Fix for introduced Mac warning related to Simplygon not being available on the platform and calling LoadModule

diff --git a/Engine/Source/Developer/MeshReductionInterface/Private/MeshReductionManagerModule.cpp b/Engine/Source/Developer/MeshReductionInterface/Private/MeshReductionManagerModule.cpp
--- a/Engine/Source/Developer/MeshReductionInterface/Private/MeshReductionManagerModule.cpp
+++ b/Engine/Source/Developer/MeshReductionInterface/Private/MeshReductionManagerModule.cpp
@@ -45,7 +45,10 @@ void FMeshReductionManagerModule::StartupModule()
 		FModuleManager::Get().LoadModule(ModuleName);
 	}
 
-	FModuleManager::Get().LoadModule("SimplygonSwarm");
+	if (FModuleManager::Get().ModuleExists(TEXT("SimplygonSwarm")))
+	{
+		FModuleManager::Get().LoadModule("SimplygonSwarm");
+	}
 	
 	TArray<IMeshReductionModule*> MeshReductionModules = IModularFeatures::Get().GetModularFeatureImplementations<IMeshReductionModule>(IMeshReductionModule::GetModularFeatureName());
 	
