Make our own HLOD and Skeletal reduction modules the default when user is not configured to use simplygon.
P1 4.22

diff --git a/Engine/Source/Developer/MeshReductionInterface/Private/MeshReductionManagerModule.cpp b/Engine/Source/Developer/MeshReductionInterface/Private/MeshReductionManagerModule.cpp
--- a/Engine/Source/Developer/MeshReductionInterface/Private/MeshReductionManagerModule.cpp
+++ b/Engine/Source/Developer/MeshReductionInterface/Private/MeshReductionManagerModule.cpp
@@ -147,6 +147,12 @@ void FMeshReductionManagerModule::StartupModule()
 		}
 	}
 
+	// Set the names that will appear as defaults in the project settings
+
+	CVarMeshReductionModule->Set(*StaticMeshModuleName);
+	CVarSkeletalMeshReductionModule->Set(*SkeletalMeshModuleName);
+	CVarProxyLODMeshReductionModule->Set(*MeshMergingModuleName);
+
 	if (!StaticMeshReduction)
 	{
 		UE_LOG(LogMeshReduction, Log, TEXT("No automatic static mesh reduction module available"));
