Fixed access to unregistered CVar

Moved some cvars out of the developer module MeshReductionInterface and to Engine where they actually belong

diff --git a/Engine/Source/Developer/MeshReductionInterface/Private/MeshReductionManagerModule.cpp b/Engine/Source/Developer/MeshReductionInterface/Private/MeshReductionManagerModule.cpp
--- a/Engine/Source/Developer/MeshReductionInterface/Private/MeshReductionManagerModule.cpp
+++ b/Engine/Source/Developer/MeshReductionInterface/Private/MeshReductionManagerModule.cpp
@@ -2,32 +2,16 @@
 
 #include "MeshReductionManagerModule.h"
 #include "IMeshReductionInterfaces.h"
-#include "HAL/IConsoleManager.h"
 #include "CoreGlobals.h"
 #include "Features/IModularFeatures.h"
 #include "Modules/ModuleManager.h"
 #include "Misc/ConfigCacheIni.h"
 
-DEFINE_LOG_CATEGORY_STATIC(LogMeshReduction, Verbose, All);
-
-static FAutoConsoleVariable CVarMeshReductionModule(
-	TEXT("r.MeshReductionModule"),
-	TEXT("QuadricMeshReduction"),
-	TEXT("Name of what mesh reduction module to choose. If blank it chooses any that exist.\n"),
-	ECVF_ReadOnly);
-
-static FAutoConsoleVariable CVarSkeletalMeshReductionModule(
-	TEXT("r.SkeletalMeshReductionModule"),
-	TEXT("SkeletalMeshReduction"),
-	TEXT("Name of what skeletal mesh reduction module to choose. If blank it chooses any that exist.\n"),
-	ECVF_ReadOnly);
-
-static FAutoConsoleVariable CVarProxyLODMeshReductionModule(
-	TEXT("r.ProxyLODMeshReductionModule"),
-	TEXT("QuadricMeshProxyLODReduction"),
-	TEXT("Name of the Proxy LOD reduction module to choose. If blank it chooses any that exist.\n"),
-	ECVF_ReadOnly);
+#include "Engine/MeshSimplificationSettings.h"
+#include "Engine/SkeletalMeshSimplificationSettings.h"
+#include "Engine/ProxyLODMeshSimplificationSettings.h"
 
+DEFINE_LOG_CATEGORY_STATIC(LogMeshReduction, Verbose, All);
 
 IMPLEMENT_MODULE(FMeshReductionManagerModule, MeshReductionInterface);
 
@@ -45,29 +29,16 @@ void FMeshReductionManagerModule::StartupModule()
 	checkf(SkeletalMeshReduction == nullptr, TEXT("Skeletal Reduction instance should be null during startup"));
 	checkf(MeshMerging           == nullptr, TEXT("Mesh Merging instance should be null during startup"));
 
-	// This module could be launched very early by static meshes loading before the settings class that stores this value has had a chance to load.  Have to read from the config file early in the startup process
-	FString MeshReductionModuleName;
-	GConfig->GetString(TEXT("/Script/Engine.MeshSimplificationSettings"), TEXT("r.MeshReductionModule"), MeshReductionModuleName, GEngineIni);
-	CVarMeshReductionModule->Set(*MeshReductionModuleName);
-
-	FString SkeletalMeshReductionModuleName;
-	GConfig->GetString(TEXT("/Script/Engine.SkeletalMeshSimplificationSettings"), TEXT("r.SkeletalMeshReductionModule"), SkeletalMeshReductionModuleName, GEngineIni);
-	// If nothing was specified, default to simplygon
-	if (SkeletalMeshReductionModuleName.IsEmpty())
-	{
-		SkeletalMeshReductionModuleName = FString("SimplygonMeshReduction");
-	}
-	CVarSkeletalMeshReductionModule->Set(*SkeletalMeshReductionModuleName);
-
-	FString HLODMeshReductionModuleName;
-	GConfig->GetString(TEXT("/Script/Engine.ProxyLODMeshSimplificationSettings"), TEXT("r.ProxyLODMeshReductionModule"), HLODMeshReductionModuleName, GEngineIni);
-	// If nothing was requested, default to simplygon for mesh merging reduction
-	if (HLODMeshReductionModuleName.IsEmpty())
-	{
-		HLODMeshReductionModuleName = FString("SimplygonMeshReduction");
-	}
-	CVarProxyLODMeshReductionModule->Set(*HLODMeshReductionModuleName);
+	// Get settings CDOs
+	UMeshSimplificationSettings* MeshSimplificationSettings_CDO = UMeshSimplificationSettings::StaticClass()->GetDefaultObject<UMeshSimplificationSettings>();
+	USkeletalMeshSimplificationSettings* SkeletalMeshSimplificationSettings_CDO = USkeletalMeshSimplificationSettings::StaticClass()->GetDefaultObject<USkeletalMeshSimplificationSettings>();
+	UProxyLODMeshSimplificationSettings* ProxyLODMeshSimplificationSettings_CDO = UProxyLODMeshSimplificationSettings::StaticClass()->GetDefaultObject<UProxyLODMeshSimplificationSettings>();
 
+	// Get configured module names.
+	FName MeshReductionModuleName = MeshSimplificationSettings_CDO->MeshReductionModuleName;
+	FName SkeletalMeshReductionModuleName = SkeletalMeshSimplificationSettings_CDO->SkeletalMeshReductionModuleName;
+	FName HLODMeshReductionModuleName = ProxyLODMeshSimplificationSettings_CDO->ProxyLODMeshReductionModuleName;
+	
 	// Retrieve reduction interfaces 
 	TArray<FName> ModuleNames;
 	FModuleManager::Get().FindModules(TEXT("*MeshReduction"), ModuleNames);
@@ -83,23 +54,19 @@ void FMeshReductionManagerModule::StartupModule()
 	
 	TArray<IMeshReductionModule*> MeshReductionModules = IModularFeatures::Get().GetModularFeatureImplementations<IMeshReductionModule>(IMeshReductionModule::GetModularFeatureName());
 	
-	const FString RequestedMeshReductionModuleName         = CVarMeshReductionModule->GetString();
-	const FString RequestedSkeletalMeshReductionModuleName = CVarSkeletalMeshReductionModule->GetString();
-	const FString RequestedProxyLODReductionModuleName     = CVarProxyLODMeshReductionModule->GetString();
-
-	// actual module names that will be used.
-	FString StaticMeshModuleName;
-	FString SkeletalMeshModuleName;
-	FString MeshMergingModuleName;
-	FString DistributedMeshMergingModuleName;
+	// Actual module names that will be used.
+	FName StaticMeshModuleName;
+	FName SkeletalMeshModuleName;
+	FName MeshMergingModuleName;
+	FName DistributedMeshMergingModuleName;
 
 	for (IMeshReductionModule* Module : MeshReductionModules)
 	{
 		// Is this a requested module?
-		const FString ModuleName = Module->GetName();
-		const bool bIsRequestedMeshReductionModule         = ModuleName.Equals(RequestedMeshReductionModuleName);
-		const bool bIsRequestedSkeletalMeshReductionModule = ModuleName.Equals(RequestedSkeletalMeshReductionModuleName);
-		const bool bIsRequestedProxyLODReductionModule     = ModuleName.Equals(RequestedProxyLODReductionModuleName);
+		const FName ModuleName(Module->GetName());
+		const bool bIsRequestedMeshReductionModule         = ModuleName == MeshReductionModuleName;
+		const bool bIsRequestedSkeletalMeshReductionModule = ModuleName == SkeletalMeshReductionModuleName;
+		const bool bIsRequestedProxyLODReductionModule     = ModuleName == HLODMeshReductionModuleName;	
 	
 
 		// Look for MeshReduction interface
@@ -147,11 +114,9 @@ void FMeshReductionManagerModule::StartupModule()
 		}
 	}
 
-	// Set the names that will appear as defaults in the project settings
-
-	CVarMeshReductionModule->Set(*StaticMeshModuleName);
-	CVarSkeletalMeshReductionModule->Set(*SkeletalMeshModuleName);
-	CVarProxyLODMeshReductionModule->Set(*MeshMergingModuleName);
+	MeshSimplificationSettings_CDO->SetMeshReductionModuleName(StaticMeshModuleName);
+	SkeletalMeshSimplificationSettings_CDO->SetSkeletalMeshReductionModuleName(SkeletalMeshModuleName);
+	ProxyLODMeshSimplificationSettings_CDO->SetProxyLODMeshReductionModuleName(MeshMergingModuleName);
 
 	if (!StaticMeshReduction)
 	{
@@ -159,7 +124,7 @@ void FMeshReductionManagerModule::StartupModule()
 	}
 	else
 	{
-		UE_LOG(LogMeshReduction, Log, TEXT("Using %s for automatic static mesh reduction"), *StaticMeshModuleName);
+		UE_LOG(LogMeshReduction, Log, TEXT("Using %s for automatic static mesh reduction"), *StaticMeshModuleName.ToString());
 	}
 
 	if (!SkeletalMeshReduction)
@@ -168,7 +133,7 @@ void FMeshReductionManagerModule::StartupModule()
 	}
 	else
 	{
-		UE_LOG(LogMeshReduction, Log, TEXT("Using %s for automatic skeletal mesh reduction"), *SkeletalMeshReductionModuleName);
+		UE_LOG(LogMeshReduction, Log, TEXT("Using %s for automatic skeletal mesh reduction"), *SkeletalMeshReductionModuleName.ToString());
 	}
 
 	if (!MeshMerging)
@@ -177,7 +142,7 @@ void FMeshReductionManagerModule::StartupModule()
 	}
 	else
 	{
-		UE_LOG(LogMeshReduction, Log, TEXT("Using %s for automatic mesh merging"), *MeshMergingModuleName);
+		UE_LOG(LogMeshReduction, Log, TEXT("Using %s for automatic mesh merging"), *MeshMergingModuleName.ToString());
 	}
 
 
@@ -187,7 +152,7 @@ void FMeshReductionManagerModule::StartupModule()
 	}
 	else
 	{
-		UE_LOG(LogMeshReduction, Log, TEXT("Using %s for distributed automatic mesh merging"), *DistributedMeshMergingModuleName);
+		UE_LOG(LogMeshReduction, Log, TEXT("Using %s for distributed automatic mesh merging"), *DistributedMeshMergingModuleName.ToString());
 	}
 }
 
diff --git a/Engine/Source/Runtime/Engine/Classes/Engine/MeshSimplificationSettings.h b/Engine/Source/Runtime/Engine/Classes/Engine/MeshSimplificationSettings.h
--- a/Engine/Source/Runtime/Engine/Classes/Engine/MeshSimplificationSettings.h
+++ b/Engine/Source/Runtime/Engine/Classes/Engine/MeshSimplificationSettings.h
@@ -8,13 +8,16 @@
 #include "MeshSimplificationSettings.generated.h"
 
 UCLASS(config=Engine, defaultconfig, meta=(DisplayName="Mesh Simplification"))
-class UMeshSimplificationSettings : public UDeveloperSettings
+class ENGINE_API UMeshSimplificationSettings : public UDeveloperSettings
 {
 	GENERATED_UCLASS_BODY()
 
 public:
 	virtual FName GetContainerName() const override;
 	virtual FName GetCategoryName() const override;
+
+	void SetMeshReductionModuleName(FName InMeshReductionModuleName);
+
 public:
 	/** Mesh reduction plugin to use when simplifying mesh geometry */
 	UPROPERTY(config, EditAnywhere, Category=General, meta=(ConsoleVariable="r.MeshReductionModule", DisplayName="Mesh Reduction Plugin", ConfigRestartRequired=true))
diff --git a/Engine/Source/Runtime/Engine/Classes/Engine/ProxyLODMeshSimplificationSettings.h b/Engine/Source/Runtime/Engine/Classes/Engine/ProxyLODMeshSimplificationSettings.h
--- a/Engine/Source/Runtime/Engine/Classes/Engine/ProxyLODMeshSimplificationSettings.h
+++ b/Engine/Source/Runtime/Engine/Classes/Engine/ProxyLODMeshSimplificationSettings.h
@@ -11,13 +11,16 @@
 * Controls the system used to generate proxy LODs with merged meshes (i.e. the HLOD system).
 */
 UCLASS(config = Engine, defaultconfig, meta = (DisplayName = "Hierarchical LOD Mesh Simplification"))
-class UProxyLODMeshSimplificationSettings : public UDeveloperSettings
+class ENGINE_API UProxyLODMeshSimplificationSettings : public UDeveloperSettings
 {
 	GENERATED_UCLASS_BODY()
 
 public:
 	virtual FName GetContainerName() const override;
 	virtual FName GetCategoryName() const override;
+
+	void SetProxyLODMeshReductionModuleName(FName InProxyLODMeshReductionModuleName);
+
 public:
 	/** Mesh reduction plugin to use when simplifying mesh geometry for Hierarchical LOD */
 	UPROPERTY(config, EditAnywhere, Category = General, meta = (ConsoleVariable = "r.ProxyLODMeshReductionModule", DisplayName = "Hierarchical LOD Mesh Reduction Plugin", ConfigRestartRequired = true))
diff --git a/Engine/Source/Runtime/Engine/Classes/Engine/SkeletalMeshSimplificationSettings.h b/Engine/Source/Runtime/Engine/Classes/Engine/SkeletalMeshSimplificationSettings.h
--- a/Engine/Source/Runtime/Engine/Classes/Engine/SkeletalMeshSimplificationSettings.h
+++ b/Engine/Source/Runtime/Engine/Classes/Engine/SkeletalMeshSimplificationSettings.h
@@ -11,13 +11,16 @@
 * Controls the selection of the system used to simplify skeletal meshes.
 */
 UCLASS(config = Engine, defaultconfig, meta = (DisplayName = "Skeletal Mesh Simplification"))
-class USkeletalMeshSimplificationSettings : public UDeveloperSettings
+class ENGINE_API USkeletalMeshSimplificationSettings : public UDeveloperSettings
 {
 	GENERATED_UCLASS_BODY()
 
 public:
 	virtual FName GetContainerName() const override;
 	virtual FName GetCategoryName() const override;
+
+	void SetSkeletalMeshReductionModuleName(FName InSkeletalMeshReductionModuleName);
+
 public:
 	/** Mesh reduction plugin to use when simplifying skeletal meshes */
 	UPROPERTY(config, EditAnywhere, Category = General, meta = (ConsoleVariable = "r.SkeletalMeshReductionModule", DisplayName = "Skeletal Mesh Reduction Plugin", ConfigRestartRequired = true))
diff --git a/Engine/Source/Runtime/Engine/Private/MeshSimplificationSettings.cpp b/Engine/Source/Runtime/Engine/Private/MeshSimplificationSettings.cpp
--- a/Engine/Source/Runtime/Engine/Private/MeshSimplificationSettings.cpp
+++ b/Engine/Source/Runtime/Engine/Private/MeshSimplificationSettings.cpp
@@ -2,8 +2,8 @@
 
 #include "Engine/MeshSimplificationSettings.h"
 #include "Modules/ModuleManager.h"
-#include "UObject/UnrealType.h"
-#include "Modules/ModuleManager.h"
+#include "HAL/IConsoleManager.h"
+#include "CoreGlobals.h"
 
 #include UE_INLINE_GENERATED_CPP_BY_NAME(MeshSimplificationSettings)
 
@@ -11,10 +11,15 @@
 #include "IMeshReductionManagerModule.h"
 #endif
 
+static FAutoConsoleVariable CVarMeshReductionModule(
+	TEXT("r.MeshReductionModule"),
+	TEXT("QuadricMeshReduction"),
+	TEXT("Name of what mesh reduction module to choose. If blank it chooses any that exist.\n"),
+	ECVF_ReadOnly);
+
 UMeshSimplificationSettings::UMeshSimplificationSettings(const FObjectInitializer& ObjectInitializer)
 	: Super(ObjectInitializer)
 {
-
 }
 
 FName UMeshSimplificationSettings::GetContainerName() const
@@ -29,6 +34,12 @@ FName UMeshSimplificationSettings::GetCategoryName() const
 	return EditorCategoryName;
 }
 
+void UMeshSimplificationSettings::SetMeshReductionModuleName(FName InMeshReductionModuleName)
+{
+	MeshReductionModuleName = InMeshReductionModuleName;
+	CVarMeshReductionModule->Set(*MeshReductionModuleName.ToString());
+}
+
 void UMeshSimplificationSettings::PostInitProperties()
 {
 	Super::PostInitProperties(); 
diff --git a/Engine/Source/Runtime/Engine/Private/ProxyLODMeshSimplificationSettings.cpp b/Engine/Source/Runtime/Engine/Private/ProxyLODMeshSimplificationSettings.cpp
--- a/Engine/Source/Runtime/Engine/Private/ProxyLODMeshSimplificationSettings.cpp
+++ b/Engine/Source/Runtime/Engine/Private/ProxyLODMeshSimplificationSettings.cpp
@@ -1,14 +1,20 @@
 // Copyright Epic Games, Inc. All Rights Reserved.
 
 #include "Engine/ProxyLODMeshSimplificationSettings.h"
-#include "UObject/UnrealType.h"
+#include "HAL/IConsoleManager.h"
+#include "CoreGlobals.h"
 
 #include UE_INLINE_GENERATED_CPP_BY_NAME(ProxyLODMeshSimplificationSettings)
 
+static FAutoConsoleVariable CVarProxyLODMeshReductionModule(
+	TEXT("r.ProxyLODMeshReductionModule"),
+	TEXT("QuadricMeshProxyLODReduction"),
+	TEXT("Name of the Proxy LOD reduction module to choose. If blank it chooses any that exist.\n"),
+	ECVF_ReadOnly);
+
 UProxyLODMeshSimplificationSettings::UProxyLODMeshSimplificationSettings(const FObjectInitializer& ObjectInitializer)
 	: Super(ObjectInitializer)
 {
-
 }
 
 FName UProxyLODMeshSimplificationSettings::GetContainerName() const
@@ -23,6 +29,12 @@ FName UProxyLODMeshSimplificationSettings::GetCategoryName() const
 	return EditorCategoryName;
 }
 
+void UProxyLODMeshSimplificationSettings::SetProxyLODMeshReductionModuleName(FName InProxyLODMeshReductionModuleName)
+{
+	ProxyLODMeshReductionModuleName = InProxyLODMeshReductionModuleName;
+	CVarProxyLODMeshReductionModule->Set(*ProxyLODMeshReductionModuleName.ToString());
+}
+
 void UProxyLODMeshSimplificationSettings::PostInitProperties()
 {
 	Super::PostInitProperties();
diff --git a/Engine/Source/Runtime/Engine/Private/SkeletalMeshSimplificationSettings.cpp b/Engine/Source/Runtime/Engine/Private/SkeletalMeshSimplificationSettings.cpp
--- a/Engine/Source/Runtime/Engine/Private/SkeletalMeshSimplificationSettings.cpp
+++ b/Engine/Source/Runtime/Engine/Private/SkeletalMeshSimplificationSettings.cpp
@@ -1,14 +1,20 @@
 // Copyright Epic Games, Inc. All Rights Reserved.
 
 #include "Engine/SkeletalMeshSimplificationSettings.h"
-#include "UObject/UnrealType.h"
+#include "HAL/IConsoleManager.h"
+#include "CoreGlobals.h"
 
 #include UE_INLINE_GENERATED_CPP_BY_NAME(SkeletalMeshSimplificationSettings)
 
+static FAutoConsoleVariable CVarSkeletalMeshReductionModule(
+	TEXT("r.SkeletalMeshReductionModule"),
+	TEXT("SkeletalMeshReduction"),
+	TEXT("Name of what skeletal mesh reduction module to choose. If blank it chooses any that exist.\n"),
+	ECVF_ReadOnly);
+
 USkeletalMeshSimplificationSettings::USkeletalMeshSimplificationSettings(const FObjectInitializer& ObjectInitializer)
 	: Super(ObjectInitializer)
 {
-
 }
 
 FName USkeletalMeshSimplificationSettings::GetContainerName() const
@@ -23,6 +29,12 @@ FName USkeletalMeshSimplificationSettings::GetCategoryName() const
 	return EditorCategoryName;
 }
 
+void USkeletalMeshSimplificationSettings::SetSkeletalMeshReductionModuleName(FName InSkeletalMeshReductionModuleName)
+{
+	SkeletalMeshReductionModuleName = InSkeletalMeshReductionModuleName;
+	CVarSkeletalMeshReductionModule->Set(*SkeletalMeshReductionModuleName.ToString());
+}
+
 void USkeletalMeshSimplificationSettings::PostInitProperties()
 {
 	Super::PostInitProperties();
