Mesh reduction modules selection aren't properly applied on engine restart

diff --git a/Engine/Source/Developer/MeshReductionInterface/Private/MeshReductionManagerModule.cpp b/Engine/Source/Developer/MeshReductionInterface/Private/MeshReductionManagerModule.cpp
--- a/Engine/Source/Developer/MeshReductionInterface/Private/MeshReductionManagerModule.cpp
+++ b/Engine/Source/Developer/MeshReductionInterface/Private/MeshReductionManagerModule.cpp
@@ -34,6 +34,25 @@ void FMeshReductionManagerModule::StartupModule()
 	USkeletalMeshSimplificationSettings* SkeletalMeshSimplificationSettings_CDO = USkeletalMeshSimplificationSettings::StaticClass()->GetDefaultObject<USkeletalMeshSimplificationSettings>();
 	UProxyLODMeshSimplificationSettings* ProxyLODMeshSimplificationSettings_CDO = UProxyLODMeshSimplificationSettings::StaticClass()->GetDefaultObject<UProxyLODMeshSimplificationSettings>();
 
+	// Update module names from config entries
+	{
+		FString MeshReductionModuleName;
+		if (GConfig->GetString(TEXT("/Script/Engine.MeshSimplificationSettings"), TEXT("r.MeshReductionModule"), MeshReductionModuleName, GEngineIni))
+		{
+			MeshSimplificationSettings_CDO->SetMeshReductionModuleName(*MeshReductionModuleName);
+		}
+		FString SkeletalMeshReductionModuleName;
+		if (GConfig->GetString(TEXT("/Script/Engine.SkeletalMeshSimplificationSettings"), TEXT("r.SkeletalMeshReductionModule"), SkeletalMeshReductionModuleName, GEngineIni))
+		{
+			SkeletalMeshSimplificationSettings_CDO->SetSkeletalMeshReductionModuleName(*SkeletalMeshReductionModuleName);
+		}
+		FString HLODMeshReductionModuleName;
+		if (GConfig->GetString(TEXT("/Script/Engine.ProxyLODMeshSimplificationSettings"), TEXT("r.ProxyLODMeshReductionModule"), HLODMeshReductionModuleName, GEngineIni))
+		{
+			ProxyLODMeshSimplificationSettings_CDO->SetProxyLODMeshReductionModuleName(*HLODMeshReductionModuleName);
+		}
+	}
+
 	// Get configured module names.
 	FName MeshReductionModuleName = MeshSimplificationSettings_CDO->MeshReductionModuleName;
 	FName SkeletalMeshReductionModuleName = SkeletalMeshSimplificationSettings_CDO->SkeletalMeshReductionModuleName;
