Mesh Reduction Manager - Improved logging

diff --git a/Engine/Source/Developer/MeshReductionInterface/Private/MeshReductionManagerModule.cpp b/Engine/Source/Developer/MeshReductionInterface/Private/MeshReductionManagerModule.cpp
--- a/Engine/Source/Developer/MeshReductionInterface/Private/MeshReductionManagerModule.cpp
+++ b/Engine/Source/Developer/MeshReductionInterface/Private/MeshReductionManagerModule.cpp
@@ -55,9 +55,23 @@ void FMeshReductionManagerModule::StartupModule()
 
 	// Get configured module names.
 	FName MeshReductionModuleName = MeshSimplificationSettings_CDO->MeshReductionModuleName;
+	if (!FModuleManager::Get().ModuleExists(*MeshReductionModuleName.ToString()))
+	{
+		UE_LOG(LogMeshReduction, Warning, TEXT("Mesh reduction module (r.MeshReductionModule) set to \"%s\" which doesn't exist."), *MeshReductionModuleName.ToString());
+	}
+
 	FName SkeletalMeshReductionModuleName = SkeletalMeshSimplificationSettings_CDO->SkeletalMeshReductionModuleName;
+	if (!FModuleManager::Get().ModuleExists(*SkeletalMeshReductionModuleName.ToString()))
+	{
+		UE_LOG(LogMeshReduction, Warning, TEXT("Skeletal mesh reduction module (r.SkeletalMeshReductionModule) set to \"%s\" which doesn't exist."), *SkeletalMeshReductionModuleName.ToString());
+	}
+
 	FName HLODMeshReductionModuleName = ProxyLODMeshSimplificationSettings_CDO->ProxyLODMeshReductionModuleName;
-	
+	if (!FModuleManager::Get().ModuleExists(*HLODMeshReductionModuleName.ToString()))
+	{
+		UE_LOG(LogMeshReduction, Warning, TEXT("HLOD mesh reduction module (r.ProxyLODMeshReductionModule) set to \"%s\" which doesn't exist."), *HLODMeshReductionModuleName.ToString());
+	}
+
 	// Retrieve reduction interfaces 
 	TArray<FName> ModuleNames;
 	FModuleManager::Get().FindModules(TEXT("*MeshReduction"), ModuleNames);
@@ -139,39 +153,39 @@ void FMeshReductionManagerModule::StartupModule()
 
 	if (!StaticMeshReduction)
 	{
-		UE_LOG(LogMeshReduction, Log, TEXT("No automatic static mesh reduction module available"));
+		UE_LOG(LogMeshReduction, Display, TEXT("No automatic static mesh reduction module available"));
 	}
 	else
 	{
-		UE_LOG(LogMeshReduction, Log, TEXT("Using %s for automatic static mesh reduction"), *StaticMeshModuleName.ToString());
+		UE_LOG(LogMeshReduction, Display, TEXT("Using %s for automatic static mesh reduction"), *StaticMeshModuleName.ToString());
 	}
 
 	if (!SkeletalMeshReduction)
 	{
-		UE_LOG(LogMeshReduction, Log, TEXT("No automatic skeletal mesh reduction module available"));
+		UE_LOG(LogMeshReduction, Display, TEXT("No automatic skeletal mesh reduction module available"));
 	}
 	else
 	{
-		UE_LOG(LogMeshReduction, Log, TEXT("Using %s for automatic skeletal mesh reduction"), *SkeletalMeshReductionModuleName.ToString());
+		UE_LOG(LogMeshReduction, Display, TEXT("Using %s for automatic skeletal mesh reduction"), *SkeletalMeshReductionModuleName.ToString());
 	}
 
 	if (!MeshMerging)
 	{
-		UE_LOG(LogMeshReduction, Log, TEXT("No automatic mesh merging module available"));
+		UE_LOG(LogMeshReduction, Display, TEXT("No automatic mesh merging module available"));
 	}
 	else
 	{
-		UE_LOG(LogMeshReduction, Log, TEXT("Using %s for automatic mesh merging"), *MeshMergingModuleName.ToString());
+		UE_LOG(LogMeshReduction, Display, TEXT("Using %s for automatic mesh merging"), *MeshMergingModuleName.ToString());
 	}
 
 
 	if (!DistributedMeshMerging)
 	{
-		UE_LOG(LogMeshReduction, Log, TEXT("No distributed automatic mesh merging module available"));
+		UE_LOG(LogMeshReduction, Display, TEXT("No distributed automatic mesh merging module available"));
 	}
 	else
 	{
-		UE_LOG(LogMeshReduction, Log, TEXT("Using %s for distributed automatic mesh merging"), *DistributedMeshMergingModuleName.ToString());
+		UE_LOG(LogMeshReduction, Display, TEXT("Using %s for distributed automatic mesh merging"), *DistributedMeshMergingModuleName.ToString());
 	}
 }
 
