MeshReductionManagerModule: Change module existence check from Warning to Display as the system will fallback on an existing module.

Current r.ProxyLODMeshReductionModule default "ProxyLODMeshReduction" doesn't exist on some platforms

#rnx

diff --git a/Engine/Source/Developer/MeshReductionInterface/Private/MeshReductionManagerModule.cpp b/Engine/Source/Developer/MeshReductionInterface/Private/MeshReductionManagerModule.cpp
--- a/Engine/Source/Developer/MeshReductionInterface/Private/MeshReductionManagerModule.cpp
+++ b/Engine/Source/Developer/MeshReductionInterface/Private/MeshReductionManagerModule.cpp
@@ -57,19 +57,19 @@ void FMeshReductionManagerModule::StartupModule()
 	FName MeshReductionModuleName = MeshSimplificationSettings_CDO->MeshReductionModuleName;
 	if (!FModuleManager::Get().ModuleExists(*MeshReductionModuleName.ToString()))
 	{
-		UE_LOG(LogMeshReduction, Warning, TEXT("Mesh reduction module (r.MeshReductionModule) set to \"%s\" which doesn't exist."), *MeshReductionModuleName.ToString());
+		UE_LOG(LogMeshReduction, Display, TEXT("Mesh reduction module (r.MeshReductionModule) set to \"%s\" which doesn't exist."), *MeshReductionModuleName.ToString());
 	}
 
 	FName SkeletalMeshReductionModuleName = SkeletalMeshSimplificationSettings_CDO->SkeletalMeshReductionModuleName;
 	if (!FModuleManager::Get().ModuleExists(*SkeletalMeshReductionModuleName.ToString()))
 	{
-		UE_LOG(LogMeshReduction, Warning, TEXT("Skeletal mesh reduction module (r.SkeletalMeshReductionModule) set to \"%s\" which doesn't exist."), *SkeletalMeshReductionModuleName.ToString());
+		UE_LOG(LogMeshReduction, Display, TEXT("Skeletal mesh reduction module (r.SkeletalMeshReductionModule) set to \"%s\" which doesn't exist."), *SkeletalMeshReductionModuleName.ToString());
 	}
 
 	FName HLODMeshReductionModuleName = ProxyLODMeshSimplificationSettings_CDO->ProxyLODMeshReductionModuleName;
 	if (!FModuleManager::Get().ModuleExists(*HLODMeshReductionModuleName.ToString()))
 	{
-		UE_LOG(LogMeshReduction, Warning, TEXT("HLOD mesh reduction module (r.ProxyLODMeshReductionModule) set to \"%s\" which doesn't exist."), *HLODMeshReductionModuleName.ToString());
+		UE_LOG(LogMeshReduction, Display, TEXT("HLOD mesh reduction module (r.ProxyLODMeshReductionModule) set to \"%s\" which doesn't exist."), *HLODMeshReductionModuleName.ToString());
 	}
 
 	// Retrieve reduction interfaces 
